fixed NaN printout if you only run CPU SynthBenchmark

diff --git a/Engine/Source/Developer/SynthBenchmark/Private/SynthBenchmarkPrivatePCH.cpp b/Engine/Source/Developer/SynthBenchmark/Private/SynthBenchmarkPrivatePCH.cpp
--- a/Engine/Source/Developer/SynthBenchmark/Private/SynthBenchmarkPrivatePCH.cpp
+++ b/Engine/Source/Developer/SynthBenchmark/Private/SynthBenchmarkPrivatePCH.cpp
@@ -201,7 +201,12 @@ void FSynthBenchmark::Run(FSynthBenchmarkResults& InOut, bool bGPUBenchmark, flo
 	}
 	
 	UE_LOG(LogSynthBenchmark, Display, TEXT("  CPUIndex: %.1f"), InOut.CPUStats->ComputePerfIndex());
-	UE_LOG(LogSynthBenchmark, Display, TEXT("  GPUIndex: %.1f"), InOut.GPUStats->ComputePerfIndex());
+
+	if(bGPUBenchmark)
+	{
+		UE_LOG(LogSynthBenchmark, Display, TEXT("  GPUIndex: %.1f"), InOut.GPUStats->ComputePerfIndex());
+	}
+
 	UE_LOG(LogSynthBenchmark, Display, TEXT(""));
 	UE_LOG(LogSynthBenchmark, Display, TEXT("         ... Total Time: %f sec"),  (float)(FPlatformTime::Seconds() - StartTime));
 }
\ No newline at end of file
